Fix for goofs with serialization.

diff --git a/Engine/Source/Developer/CookMetadata/Private/CookMetadata.cpp b/Engine/Source/Developer/CookMetadata/Private/CookMetadata.cpp
--- a/Engine/Source/Developer/CookMetadata/Private/CookMetadata.cpp
+++ b/Engine/Source/Developer/CookMetadata/Private/CookMetadata.cpp
@@ -31,15 +31,25 @@ bool FCookMetadataState::Serialize(FArchive& Ar)
 	{
 		MagicHeader = COOK_METADATA_HEADER_MAGIC;
 		Ar << MagicHeader;
+
+		Version = ECookMetadataStateVersion::LatestVersion;
 	}
+
 	Ar << Version;
+	if (Ar.IsLoading())
+	{
+		if (Version != ECookMetadataStateVersion::LatestVersion)
+		{
+			return false; // invalid version - current we don't support backcompat
+		}
+	}
 	Ar << PluginHierarchy;
 	Ar << AssociatedDevelopmentAssetRegistryHash;
 	Ar << AssociatedDevelopmentAssetRegistryHashPostWriteback;
 	Ar << Platform;
 	Ar << BuildVersion;
 	Ar << HordeJobId;
-
+	Ar << SizesPresent;
 	return true;
 }
 
diff --git a/Engine/Source/Developer/CookMetadata/Public/CookMetadata.h b/Engine/Source/Developer/CookMetadata/Public/CookMetadata.h
--- a/Engine/Source/Developer/CookMetadata/Public/CookMetadata.h
+++ b/Engine/Source/Developer/CookMetadata/Public/CookMetadata.h
@@ -16,6 +16,7 @@ enum class ECookMetadataStateVersion : uint8
 	InvalidVersion = 0,
 	PluginHierarchy = 1,
 	PostWritebackHash = 2,
+	FixSerialization = 3,
 
 	// Add new versions above this.
 	VersionCount,
@@ -249,7 +250,7 @@ private:
 	FString HordeJobId;
 
 	// Updated by unrealpak when plugin size information is added.
-	ECookMetadataSizesPresent SizesPresent;
+	ECookMetadataSizesPresent SizesPresent = ECookMetadataSizesPresent::NotPresent;
 };
 
 COOKMETADATA_API const FString& GetCookMetadataFilename();
