Adding a helper to convert UE::Cook::ECookMetadataSizesPresent to an FText

diff --git a/Engine/Source/Developer/CookMetadata/Private/CookMetadata.cpp b/Engine/Source/Developer/CookMetadata/Private/CookMetadata.cpp
--- a/Engine/Source/Developer/CookMetadata/Private/CookMetadata.cpp
+++ b/Engine/Source/Developer/CookMetadata/Private/CookMetadata.cpp
@@ -3,6 +3,7 @@
 #include "CookMetadata.h"
 
 #include "Hash/xxhash.h"
+#include "Internationalization/Internationalization.h"
 #include "Memory/MemoryView.h"
 
 namespace UE::Cook
@@ -59,5 +60,18 @@ uint64 FCookMetadataState::ComputeHashOfDevelopmentAssetRegistry(FMemoryView InS
 	return FXxHash64::HashBufferChunked(InSerializedDevelopmentAssetRegistry.GetData(), InSerializedDevelopmentAssetRegistry.GetSize(), 1ULL << 19).Hash;
 }
 
+FText FCookMetadataState::GetSizesPresentAsText() const
+{
+	// Uncapitalized, presentation-ready
+	static FText CookMetadataSizesPresentStrings[] =
+	{
+		NSLOCTEXT("CookMetadata", "CookMetadataNotPresent", "not present"),
+		NSLOCTEXT("CookMetadata", "CookMetadataCompressed", "compressed"),
+		NSLOCTEXT("CookMetadata", "CookMetadataUncompressed", "uncompressed")
+	};
+
+	static_assert(sizeof(CookMetadataSizesPresentStrings) / sizeof(CookMetadataSizesPresentStrings[0]) == static_cast<size_t>(ECookMetadataSizesPresent::Count));
+	return CookMetadataSizesPresentStrings[static_cast<size_t>(SizesPresent)];
+}
 
 } // end namespace
\ No newline at end of file
diff --git a/Engine/Source/Developer/CookMetadata/Public/CookMetadata.h b/Engine/Source/Developer/CookMetadata/Public/CookMetadata.h
--- a/Engine/Source/Developer/CookMetadata/Public/CookMetadata.h
+++ b/Engine/Source/Developer/CookMetadata/Public/CookMetadata.h
@@ -107,9 +107,10 @@ enum class ECookMetadataSizesPresent
 	Compressed,
 
 	// the selected platform isn't compressed by unrealpak, or package compression was disabled.
-	Uncompressed
-};
+	Uncompressed,
 
+	Count
+};
 
 /** The name and dependency information for a plugin that was enabled during cooking. */
 struct COOKMETADATA_API FCookMetadataPluginEntry
@@ -228,6 +229,7 @@ public:
 
 	// Returns what size information is present in FCookMetadataPluginEntry. This varies based on the platform
 	// and settings. 
+	FText GetSizesPresentAsText() const;
 	ECookMetadataSizesPresent GetSizesPresent() const { return SizesPresent; }
 	void SetSizesPresent(ECookMetadataSizesPresent InSizesPresent) { SizesPresent = InSizesPresent; }
 private:
