Fixup cookmetadata.ucookmeta after we join the cook outputs. Update the ucookedmeta file to have the same hash as the newly made combined devAR.
Added Save and Load functions to the cooked meta file

diff --git a/Engine/Plugins/Editor/AssetManagerEditor/Source/AssetManagerEditor/Private/TreeView/SAssetTableTreeView.cpp b/Engine/Plugins/Editor/AssetManagerEditor/Source/AssetManagerEditor/Private/TreeView/SAssetTableTreeView.cpp
--- a/Engine/Plugins/Editor/AssetManagerEditor/Source/AssetManagerEditor/Private/TreeView/SAssetTableTreeView.cpp
+++ b/Engine/Plugins/Editor/AssetManagerEditor/Source/AssetManagerEditor/Private/TreeView/SAssetTableTreeView.cpp
@@ -1052,15 +1052,10 @@ bool HashRegistryFile(const FString& FilePath, uint64* HashOut)
 
 bool LoadCookMetadata(const FString& FilePath, UE::Cook::FCookMetadataState& OutMetadataState)
 {
-	bool Success = false;
-	FArrayReader SerializedAssetData;
-	if (FFileHelper::LoadFileToArray(SerializedAssetData, *FilePath))
+	const bool Success = OutMetadataState.ReadFromFile(FilePath);
+	if (!Success)
 	{
-		Success = OutMetadataState.Serialize(SerializedAssetData);
-		if (!Success)
-		{
-			OutMetadataState.Reset();
-		}
+		OutMetadataState.Reset();
 	}
 	return Success;
 }
diff --git a/Engine/Source/Developer/CookMetadata/Private/CookMetadata.cpp b/Engine/Source/Developer/CookMetadata/Private/CookMetadata.cpp
--- a/Engine/Source/Developer/CookMetadata/Private/CookMetadata.cpp
+++ b/Engine/Source/Developer/CookMetadata/Private/CookMetadata.cpp
@@ -2,9 +2,13 @@
 
 #include "CookMetadata.h"
 
+#include "HAL/FileManager.h"
 #include "Hash/xxhash.h"
 #include "Internationalization/Internationalization.h"
 #include "Memory/MemoryView.h"
+#include "Misc/FileHelper.h"
+#include "Serialization/ArrayWriter.h"
+#include "Serialization/LargeMemoryReader.h"
 
 namespace UE::Cook
 {
@@ -17,6 +21,8 @@ const FString& GetCookMetadataFilename()
 
 constexpr uint32 COOK_METADATA_HEADER_MAGIC = 'UCMT';
 
+DEFINE_LOG_CATEGORY_STATIC(LogCookedMetadata, Log, All)
+
 bool FCookMetadataState::Serialize(FArchive& Ar)
 {
 	uint32 MagicHeader = 0;
@@ -54,6 +60,48 @@ bool FCookMetadataState::Serialize(FArchive& Ar)
 	return true;
 }
 
+bool FCookMetadataState::ReadFromFile(const FString& FilePath)
+{
+	TUniquePtr<FArchive> FileReader(IFileManager::Get().CreateFileReader(*FilePath));
+	if (FileReader)
+	{
+		TArray64<uint8> Data;
+		Data.SetNumUninitialized(FileReader->TotalSize());
+		FileReader->Serialize(Data.GetData(), Data.Num());
+		check(!FileReader->IsError());
+
+		FLargeMemoryReader MemoryReader(Data.GetData(), Data.Num());
+		if (Serialize(MemoryReader))
+		{
+			return true;
+		}
+		else
+		{
+			UE_LOG(LogCookedMetadata, Error, TEXT("Failed to serialize cook metadata from file (%s)"), *FilePath);
+		}
+	}
+	else
+	{
+		UE_LOG(LogCookedMetadata, Error, TEXT("Failed to make file reader from (%s)"), *FilePath);
+	}
+	return false;
+}
+
+bool FCookMetadataState::SaveToFile(const FString& FilePath)
+{
+	FArrayWriter SerializedCookMetadata;
+	Serialize(SerializedCookMetadata);
+	if (FFileHelper::SaveArrayToFile(SerializedCookMetadata, *FilePath))
+	{
+		return true;
+	}
+	else
+	{
+		UE_LOG(LogCookedMetadata, Error, TEXT("Failed to write cook metadata file (%s)"), *FilePath);
+		return false;
+	}
+}
+
 /* static */ 
 uint64 FCookMetadataState::ComputeHashOfDevelopmentAssetRegistry(FMemoryView InSerializedDevelopmentAssetRegistry)
 {
diff --git a/Engine/Source/Developer/CookMetadata/Public/CookMetadata.h b/Engine/Source/Developer/CookMetadata/Public/CookMetadata.h
--- a/Engine/Source/Developer/CookMetadata/Public/CookMetadata.h
+++ b/Engine/Source/Developer/CookMetadata/Public/CookMetadata.h
@@ -230,6 +230,9 @@ public:
 	void Reset() { *this = FCookMetadataState(); }
 	
 	bool Serialize(FArchive& Ar);
+	bool ReadFromFile(const FString& FilePath);
+	bool SaveToFile(const FString& FilePath);
+
 
 	// Plugin hierarchy information
 	void SetPluginHierarchyInfo(FCookMetadataPluginHierarchy&& InPluginHierarchy) { PluginHierarchy = MoveTemp(InPluginHierarchy); }
diff --git a/Engine/Source/Developer/IoStoreUtilities/Private/IoStoreUtilities.cpp b/Engine/Source/Developer/IoStoreUtilities/Private/IoStoreUtilities.cpp
--- a/Engine/Source/Developer/IoStoreUtilities/Private/IoStoreUtilities.cpp
+++ b/Engine/Source/Developer/IoStoreUtilities/Private/IoStoreUtilities.cpp
@@ -3978,18 +3978,11 @@ static ECookMetadataFiles FindAndLoadMetadataFiles(
 	{
 		// The cook metadata file should be adjacent to the development asset registry.
 		FString CookMetadataFileName = FPaths::GetPath(PossibleAssetRegistryFiles[0]) / UE::Cook::GetCookMetadataFilename();
-		TUniquePtr<FArchive> FileReader(IFileManager::Get().CreateFileReader(*CookMetadataFileName));
-		if (FileReader)
+		if (IFileManager::Get().FileExists(*CookMetadataFileName))
 		{
-			TArray64<uint8> Data;
-			Data.SetNumUninitialized(FileReader->TotalSize());
-			FileReader->Serialize(Data.GetData(), Data.Num());
-			check(!FileReader->IsError());
-
-			FLargeMemoryReader MemoryReader(Data.GetData(), Data.Num());
-			if (OutCookMetadata->Serialize(MemoryReader) == false)
+			if (OutCookMetadata->ReadFromFile(CookMetadataFileName) == false)
 			{
-				UE_LOG(LogIoStore, Error, TEXT("Failed to deserialize cook metadata file - loaded successfully, but invalid data. [%s]"), *CookMetadataFileName);
+				UE_LOG(LogIoStore, Error, TEXT("Failed to deserialize cook metadata file - invalid data. [%s]"), *CookMetadataFileName);
 				if (EnumHasAnyFlags(InRequiredFiles, ECookMetadataFiles::CookMetadata))
 				{
 					return ECookMetadataFiles::None;
diff --git a/Engine/Source/Editor/UnrealEd/Private/CookOnTheFlyServer.cpp b/Engine/Source/Editor/UnrealEd/Private/CookOnTheFlyServer.cpp
--- a/Engine/Source/Editor/UnrealEd/Private/CookOnTheFlyServer.cpp
+++ b/Engine/Source/Editor/UnrealEd/Private/CookOnTheFlyServer.cpp
@@ -9689,12 +9689,7 @@ void UCookOnTheFlyServer::WriteCookMetadata(const ITargetPlatform* InTargetPlatf
 			MetadataState.SetPlatformAndBuildVersion(PlatformNameString, FApp::GetBuildVersion());
 			MetadataState.SetHordeJobId(FPlatformMisc::GetEnvironmentVariable(TEXT("UE_HORDE_JOBID")));
 
-			FArrayWriter SerializedCookMetadata;
-			MetadataState.Serialize(SerializedCookMetadata);
-			if (FFileHelper::SaveArrayToFile(SerializedCookMetadata, *GetCookedCookMetadataFilename(PlatformNameString)) == false)
-			{
-				UE_LOG(LogCook, Error, TEXT("Failed to write cook metadata file (%s)"), *GetCookedCookMetadataFilename(PlatformNameString));
-			}
+			MetadataState.SaveToFile(GetCookedCookMetadataFilename(PlatformNameString));
 		}
 	}
 }
