Adding some sanity against uninitialized values and fixing the missing serialize.

diff --git a/Engine/Source/Developer/CookMetadata/Public/CookMetadata.h b/Engine/Source/Developer/CookMetadata/Public/CookMetadata.h
--- a/Engine/Source/Developer/CookMetadata/Public/CookMetadata.h
+++ b/Engine/Source/Developer/CookMetadata/Public/CookMetadata.h
@@ -21,6 +21,7 @@ enum class ECookMetadataStateVersion : uint8
 	FixSerialization = 3,
 	AddedCustomFields = 4,
 	AddedShaderPseudoHierarchy = 5,
+	AddedPluginEntryType = 6,
 
 	// Add new versions above this.
 	VersionCount,
@@ -119,6 +120,9 @@ enum class ECookMetadataSizesPresent
 
 enum class ECookMetadataPluginType
 {
+	// For sanity tracking. All types _should_ be assigned when they are added.
+	Unassigned,
+		
 	Normal,
 
 	// Root plugins are used to separate and classify game "modes" within a single project.
@@ -133,7 +137,9 @@ enum class ECookMetadataPluginType
 	// When a shader is referenced by multiple plugins then it has no natural home for assigning
 	// its size. Instead we create a set of shader pseudo plugins based on the set of root plugins
 	// referencing the shader, including possibly an "Unrooted" plugin.
-	ShaderPseudo
+	ShaderPseudo,
+
+
 };
 
 /** The name and dependency information for a plugin that was enabled during cooking. */
@@ -141,7 +147,7 @@ struct COOKMETADATA_API FCookMetadataPluginEntry
 {
 	FString Name;
 
-	ECookMetadataPluginType Type;
+	ECookMetadataPluginType Type = ECookMetadataPluginType::Unassigned;
 
 	// These contain values pulled from the uplugin json file and hold fields that are not
 	// part of the engine FPluginDescriptor. They are for per-project values. The keys for the maps
@@ -175,7 +181,7 @@ struct COOKMETADATA_API FCookMetadataPluginEntry
 	{
 		Ar << Entry.Name << Entry.DependencyIndexStart << Entry.DependencyIndexEnd;
 		Ar << Entry.InclusiveSizes << Entry.ExclusiveSizes;
-		Ar << Entry.CustomBoolFields << Entry.CustomStringFields;
+		Ar << Entry.CustomBoolFields << Entry.CustomStringFields << Entry.Type;
 		return Ar;
 	}
 };
diff --git a/Engine/Source/Developer/IoStoreUtilities/Private/IoStoreUtilities.cpp b/Engine/Source/Developer/IoStoreUtilities/Private/IoStoreUtilities.cpp
--- a/Engine/Source/Developer/IoStoreUtilities/Private/IoStoreUtilities.cpp
+++ b/Engine/Source/Developer/IoStoreUtilities/Private/IoStoreUtilities.cpp
@@ -3634,6 +3634,10 @@ static void InsertShadersInPluginHierarchy(UE::Cook::FCookMetadataState& InCookM
 	TArray<UE::Cook::FCookMetadataPluginEntry> PluginEntries = PluginHierarchy.PluginsEnabledAtCook;
 	for (int32 PluginIndex = 0; PluginIndex < PluginEntries.Num(); PluginIndex++)
 	{
+		if (PluginEntries[PluginIndex].Type == UE::Cook::ECookMetadataPluginType::Unassigned)
+		{
+			UE_LOG(LogIoStore, Warning, TEXT("Found unassigned plugin type in cook metadata! %s"), *PluginEntries[PluginIndex].Name);
+		}
 		if (PluginEntries[PluginIndex].Type == UE::Cook::ECookMetadataPluginType::ShaderPseudo)
 		{
 			// We can do a swap because we insert these at the end in a group so there shouldn't
@@ -3711,6 +3715,15 @@ static void InsertShadersInPluginHierarchy(UE::Cook::FCookMetadataState& InCookM
 	UE::Cook::FCookMetadataPluginHierarchy& MutablePluginHierarchy = InCookMetadata.GetMutablePluginHierarchy();
 	MutablePluginHierarchy.PluginDependencies = MoveTemp(DependencyList);
 	MutablePluginHierarchy.PluginsEnabledAtCook = PluginEntries;
+
+	// Sanity check we assigned plugin types
+	for (UE::Cook::FCookMetadataPluginEntry& Entry : MutablePluginHierarchy.PluginsEnabledAtCook)
+	{
+		if (Entry.Type == UE::Cook::ECookMetadataPluginType::Unassigned)
+		{
+			UE_LOG(LogIoStore, Warning, TEXT("We caused an unassigned plugin type in shader pseudo plugin generation! %s"), *Entry.Name);
+		}
+	}
 }
 
 /**
diff --git a/Engine/Source/Editor/UnrealEd/Private/CookOnTheFlyServer.cpp b/Engine/Source/Editor/UnrealEd/Private/CookOnTheFlyServer.cpp
--- a/Engine/Source/Editor/UnrealEd/Private/CookOnTheFlyServer.cpp
+++ b/Engine/Source/Editor/UnrealEd/Private/CookOnTheFlyServer.cpp
@@ -9719,11 +9719,22 @@ void UCookOnTheFlyServer::WriteCookMetadata(const ITargetPlatform* InTargetPlatf
 		UE::Cook::FCookMetadataState MetadataState;
 		UE::Cook::FCookMetadataPluginHierarchy PluginHierarchy;
 
+
 		PluginHierarchy.PluginsEnabledAtCook = MoveTemp(PluginsToAdd);
 		PluginHierarchy.PluginDependencies = MoveTemp(PluginChildArray);
 		PluginHierarchy.RootPlugins = MoveTemp(RootPlugins);
 		PluginHierarchy.CustomFieldNames = MoveTemp(CustomFieldNames);
 
+
+		// Sanity check we assigned plugin types
+		for (UE::Cook::FCookMetadataPluginEntry& Entry : PluginHierarchy.PluginsEnabledAtCook)
+		{
+			if (Entry.Type == UE::Cook::ECookMetadataPluginType::Unassigned)
+			{
+				UE_LOG(LogCook, Warning, TEXT("Found unassigned plugin type in cook metadata generation: %s"), *Entry.Name);
+			}
+		}
+
 		MetadataState.SetPluginHierarchyInfo(MoveTemp(PluginHierarchy));
 		MetadataState.SetAssociatedDevelopmentAssetRegistryHash(InDevelopmentAssetRegistryHash);
 
