log when we hit a version mismatch

diff --git a/Engine/Source/Developer/CookMetadata/Private/CookMetadata.cpp b/Engine/Source/Developer/CookMetadata/Private/CookMetadata.cpp
--- a/Engine/Source/Developer/CookMetadata/Private/CookMetadata.cpp
+++ b/Engine/Source/Developer/CookMetadata/Private/CookMetadata.cpp
@@ -47,6 +47,7 @@ bool FCookMetadataState::Serialize(FArchive& Ar)
 	{
 		if (Version != ECookMetadataStateVersion::LatestVersion)
 		{
+			UE_LOG(LogCookedMetadata, Error, TEXT("Cook metadata version mismatch: found %d, we can load %d"), Version, ECookMetadataStateVersion::LatestVersion);
 			return false; // invalid version - current we don't support backcompat
 		}
 	}
