Actually serialize the pseudo shader hierarchy.

diff --git a/Engine/Source/Developer/CookMetadata/Private/CookMetadata.cpp b/Engine/Source/Developer/CookMetadata/Private/CookMetadata.cpp
--- a/Engine/Source/Developer/CookMetadata/Private/CookMetadata.cpp
+++ b/Engine/Source/Developer/CookMetadata/Private/CookMetadata.cpp
@@ -42,14 +42,23 @@ bool FCookMetadataState::Serialize(FArchive& Ar)
 		Version = ECookMetadataStateVersion::LatestVersion;
 	}
 
+	bool bSerializeShaderHierarchy = true;
 	Ar << Version;
 	if (Ar.IsLoading())
 	{
-		if (Version != ECookMetadataStateVersion::LatestVersion)
+		if (Version < ECookMetadataStateVersion::AddedPluginEntryType)
 		{
-			UE_LOG(LogCookedMetadata, Error, TEXT("Cook metadata version mismatch: found %d, we can load %d"), Version, ECookMetadataStateVersion::LatestVersion);
+			UE_LOG(LogCookedMetadata, Error, TEXT("Cook metadata version too old: found %d, we can load %d, latest is %d"), Version, ECookMetadataStateVersion::AddedPluginEntryType, ECookMetadataStateVersion::LatestVersion);
 			return false; // invalid version - current we don't support backcompat
 		}
+		if (Version > ECookMetadataStateVersion::LatestVersion)
+		{
+			UE_LOG(LogCookedMetadata, Error, TEXT("Cook metadata version too new: found %d, latest is %d"), Version, ECookMetadataStateVersion::LatestVersion);
+			return false;
+		}
+
+		bSerializeShaderHierarchy = Version >= ECookMetadataStateVersion::ActualAddShaderPseudoHierarchy;
+
 	}
 	Ar << PluginHierarchy;
 	Ar << AssociatedDevelopmentAssetRegistryHash;
@@ -58,6 +67,11 @@ bool FCookMetadataState::Serialize(FArchive& Ar)
 	Ar << BuildVersion;
 	Ar << HordeJobId;
 	Ar << SizesPresent;
+
+	if (bSerializeShaderHierarchy)
+	{
+		Ar << ShaderPseudoHierarchy;
+	}
 	return true;
 }
 
diff --git a/Engine/Source/Developer/CookMetadata/Public/CookMetadata.h b/Engine/Source/Developer/CookMetadata/Public/CookMetadata.h
--- a/Engine/Source/Developer/CookMetadata/Public/CookMetadata.h
+++ b/Engine/Source/Developer/CookMetadata/Public/CookMetadata.h
@@ -20,8 +20,11 @@ enum class ECookMetadataStateVersion : uint8
 	PostWritebackHash = 2,
 	FixSerialization = 3,
 	AddedCustomFields = 4,
+
+	// Note: This is a lie, the shader data wasn't actually serializing to file, use ActualAddShaderPseudoHierarchy
 	AddedShaderPseudoHierarchy = 5,
 	AddedPluginEntryType = 6,
+	ActualAddShaderPseudoHierarchy = 7,
 
 	// Add new versions above this.
 	VersionCount,
@@ -257,6 +260,12 @@ struct COOKMETADATA_API FCookMetadataShaderPseudoAsset
 
 	// Shaders are always compressed, independent of what GetSizesPresent() says.
 	uint32 CompressedSize = 0;
+
+	friend FArchive& operator<<(FArchive& Ar, FCookMetadataShaderPseudoAsset& PseudoAsset)
+	{
+		Ar << PseudoAsset.Name << PseudoAsset.CompressedSize;
+		return Ar;
+	}
 };
 
 struct COOKMETADATA_API FCookMetadataShaderPseudoHierarchy
@@ -270,6 +279,12 @@ struct COOKMETADATA_API FCookMetadataShaderPseudoHierarchy
 	// Keyed off a PackageName in the project, returns a [start, end) pair of indices
 	// into DependencyList for the shaders that package depends on.
 	TMap<FName, TPair<int32, int32>> PackageShaderDependencyMap;
+
+	friend FArchive& operator<<(FArchive& Ar, FCookMetadataShaderPseudoHierarchy& PseudoHierarchy)
+	{
+		Ar << PseudoHierarchy.ShaderAssets << PseudoHierarchy.DependencyList << PseudoHierarchy.PackageShaderDependencyMap;
+		return Ar;
+	}
 };
 
 /**
