Rework the custom fields in the cook metadata to make it easier to use.

diff --git a/Engine/Source/Developer/CookMetadata/Private/CookMetadata.cpp b/Engine/Source/Developer/CookMetadata/Private/CookMetadata.cpp
--- a/Engine/Source/Developer/CookMetadata/Private/CookMetadata.cpp
+++ b/Engine/Source/Developer/CookMetadata/Private/CookMetadata.cpp
@@ -23,6 +23,105 @@ constexpr uint32 COOK_METADATA_HEADER_MAGIC = 'UCMT';
 
 DEFINE_LOG_CATEGORY_STATIC(LogCookedMetadata, Log, All)
 
+
+// For backwards compatibility - remove this after 5.4 is released.
+
+/** The name and dependency information for a plugin that was enabled during cooking. */
+struct COOKMETADATA_API FCookMetadataPluginEntry_ActualAddShaderPseudoHierarchy
+{
+	FString Name;
+	ECookMetadataPluginType Type = ECookMetadataPluginType::Unassigned;
+	TMap<uint8, bool> CustomBoolFields;
+	TMap<uint8, FString> CustomStringFields;
+	uint32 DependencyIndexStart = 0;
+	uint32 DependencyIndexEnd = 0;
+	FPluginSizeInfo InclusiveSizes;
+	FPluginSizeInfo ExclusiveSizes;
+
+	friend FArchive& operator<<(FArchive& Ar, FCookMetadataPluginEntry_ActualAddShaderPseudoHierarchy& Entry)
+	{
+		Ar << Entry.Name << Entry.DependencyIndexStart << Entry.DependencyIndexEnd;
+		Ar << Entry.InclusiveSizes << Entry.ExclusiveSizes;
+		Ar << Entry.CustomBoolFields << Entry.CustomStringFields << Entry.Type;
+		return Ar;
+	}
+};
+
+struct FCookMetadataPluginHierarchy_ActualAddShaderPseudoHierarchy
+{
+	TArray<FCookMetadataPluginEntry_ActualAddShaderPseudoHierarchy> PluginsEnabledAtCook;
+	TArray<uint16> PluginDependencies;
+	TArray<uint16> RootPlugins;
+	TArray<FString> CustomFieldNames;
+
+	friend FArchive& operator<<(FArchive& Ar, FCookMetadataPluginHierarchy_ActualAddShaderPseudoHierarchy& Hierarchy)
+	{
+		Ar << Hierarchy.PluginsEnabledAtCook << Hierarchy.PluginDependencies;
+		Ar << Hierarchy.RootPlugins << Hierarchy.CustomFieldNames;
+		return Ar;
+	}
+};
+
+static void FCookMetadataPluginHierarchy_ActualAddShaderPseudoHierarchy_To_AdjustCustomFieldLayout(
+	FCookMetadataPluginHierarchy_ActualAddShaderPseudoHierarchy& From,
+	FCookMetadataPluginHierarchy& To)
+{
+	To.PluginDependencies = MoveTemp(From.PluginDependencies);
+	To.RootPlugins = MoveTemp(From.RootPlugins);
+	
+	// Now bring over the plugins themselves.
+	TMap<uint8, bool> TypeIsBool;
+	for (FCookMetadataPluginEntry_ActualAddShaderPseudoHierarchy& OldEntry : From.PluginsEnabledAtCook)
+	{
+		FCookMetadataPluginEntry& NewEntry = To.PluginsEnabledAtCook.AddDefaulted_GetRef();
+		NewEntry.Name = MoveTemp(OldEntry.Name);
+		NewEntry.Type = OldEntry.Type;
+
+		for (TPair<uint8, bool>& Bools : OldEntry.CustomBoolFields)
+		{
+			FCookMetadataPluginEntry::CustomFieldVariantType V;
+			V.Set<bool>(Bools.Value);
+			NewEntry.CustomFields.Add(Bools.Key, MoveTemp(V));
+			TypeIsBool.Add(Bools.Key, true);
+		}
+		for (TPair<uint8, FString>& Strings : OldEntry.CustomStringFields)
+		{
+			FCookMetadataPluginEntry::CustomFieldVariantType V;
+			V.Emplace<FString>(MoveTemp(Strings.Value));
+			NewEntry.CustomFields.Add(Strings.Key, MoveTemp(V));
+			TypeIsBool.Add(Strings.Key, false);
+		}
+
+		NewEntry.DependencyIndexStart = OldEntry.DependencyIndexStart;
+		NewEntry.DependencyIndexEnd = OldEntry.DependencyIndexEnd;
+		NewEntry.InclusiveSizes = OldEntry.InclusiveSizes;
+		NewEntry.ExclusiveSizes = OldEntry.ExclusiveSizes;
+	}
+
+	uint8 EntryIndex = 0;
+	for (FString& FieldName : From.CustomFieldNames)
+	{
+		FCookMetadataPluginHierarchy::FCustomFieldEntry Field = To.CustomFieldEntries.AddDefaulted_GetRef();
+		Field.Name = MoveTemp(FieldName);
+		Field.Type = ECookMetadataCustomFieldType::Unknown;
+
+		const bool* bIsBool = TypeIsBool.Find(EntryIndex);
+		if (bIsBool)
+		{
+			if (*bIsBool)
+			{
+				Field.Type = ECookMetadataCustomFieldType::Bool;
+			}
+			else
+			{
+				Field.Type = ECookMetadataCustomFieldType::String;
+			}
+		}
+
+		EntryIndex++;
+	}
+}
+
 bool FCookMetadataState::Serialize(FArchive& Ar)
 {
 	uint32 MagicHeader = 0;
@@ -60,7 +159,22 @@ bool FCookMetadataState::Serialize(FArchive& Ar)
 		bSerializeShaderHierarchy = Version >= ECookMetadataStateVersion::ActualAddShaderPseudoHierarchy;
 
 	}
-	Ar << PluginHierarchy;
+
+	if (Ar.IsLoading() && Version <= ECookMetadataStateVersion::ActualAddShaderPseudoHierarchy)
+	{
+		// Serialize the old hierarchy and convert to the new layout.
+		FCookMetadataPluginHierarchy_ActualAddShaderPseudoHierarchy OldHierarchy;
+		Ar << OldHierarchy;
+
+		FCookMetadataPluginHierarchy_ActualAddShaderPseudoHierarchy_To_AdjustCustomFieldLayout(
+			OldHierarchy, PluginHierarchy);
+	}
+	else
+	{
+		// saving, or on a version with the new setup.
+		Ar << PluginHierarchy;
+	}
+
 	Ar << AssociatedDevelopmentAssetRegistryHash;
 	Ar << AssociatedDevelopmentAssetRegistryHashPostWriteback;
 	Ar << Platform;
diff --git a/Engine/Source/Developer/CookMetadata/Public/CookMetadata.h b/Engine/Source/Developer/CookMetadata/Public/CookMetadata.h
--- a/Engine/Source/Developer/CookMetadata/Public/CookMetadata.h
+++ b/Engine/Source/Developer/CookMetadata/Public/CookMetadata.h
@@ -7,6 +7,7 @@
 #include "Containers/UnrealString.h"
 #include "CoreTypes.h"
 #include "Memory/MemoryFwd.h"
+#include "Misc/TVariant.h"
 #include "UObject/NameTypes.h"
 
 
@@ -25,6 +26,7 @@ enum class ECookMetadataStateVersion : uint8
 	AddedShaderPseudoHierarchy = 5,
 	AddedPluginEntryType = 6,
 	ActualAddShaderPseudoHierarchy = 7,
+	AdjustCustomFieldLayout = 8,
 
 	// Add new versions above this.
 	VersionCount,
@@ -145,6 +147,15 @@ enum class ECookMetadataPluginType
 
 };
 
+enum class ECookMetadataCustomFieldType : uint8
+{
+	Unknown,	// This only happens in the upgrade from an older cook metadata
+				// when there's nowhere to get the information from. This also means
+				// that there are no plugins that use this field so you can safely ignore it.
+	Bool,
+	String
+};
+
 /** The name and dependency information for a plugin that was enabled during cooking. */
 struct COOKMETADATA_API FCookMetadataPluginEntry
 {
@@ -154,9 +165,10 @@ struct COOKMETADATA_API FCookMetadataPluginEntry
 
 	// These contain values pulled from the uplugin json file and hold fields that are not
 	// part of the engine FPluginDescriptor. They are for per-project values. The keys for the maps
-	// are indices in to FCookMetadataPluginHierarchy::CustomFieldNames. See the comment for CustomFieldNames.
-	TMap<uint8, bool> CustomBoolFields;
-	TMap<uint8, FString> CustomStringFields;
+	// are indices in to FCookMetadataPluginHierarchy::CustomFieldEntries, where you'll also find what the type
+	// is. See the comment for CustomFieldEntries.
+	typedef TVariant<bool, FString> CustomFieldVariantType;
+	TMap<uint8, CustomFieldVariantType> CustomFields;
 
 	// The dependencies are stored in the FCookMetadataPluginHierarchy::PluginDependencies array,
 	// and this is an index into it. From there you can get a further index into PluginsEnabledAtCook
@@ -180,11 +192,12 @@ struct COOKMETADATA_API FCookMetadataPluginEntry
 
 	uint32 DependencyCount() const { return DependencyIndexEnd - DependencyIndexStart; }
 
+	// !!! If you edit this, be sure to update the upgrade paths in CookMetadata.cpp
 	friend FArchive& operator<<(FArchive& Ar, FCookMetadataPluginEntry& Entry)
 	{
 		Ar << Entry.Name << Entry.DependencyIndexStart << Entry.DependencyIndexEnd;
 		Ar << Entry.InclusiveSizes << Entry.ExclusiveSizes;
-		Ar << Entry.CustomBoolFields << Entry.CustomStringFields << Entry.Type;
+		Ar << Entry.CustomFields << Entry.Type;
 		return Ar;
 	}
 };
@@ -232,12 +245,22 @@ struct COOKMETADATA_API FCookMetadataPluginHierarchy
 	// Note that in the per platform case, the base value is not required (i.e. ExampleProjectBool above), however it
 	// provides a default value for unlisted platforms. If it does not exist, false is used for bool fields, and an empty string
 	// for string fields.
-	TArray<FString> CustomFieldNames;
+	struct FCustomFieldEntry
+	{
+		FString Name;
+		ECookMetadataCustomFieldType Type;
+		friend FArchive& operator<<(FArchive& Ar, FCustomFieldEntry& Entry)
+		{
+			return Ar << Entry.Name << Entry.Type;
+		}
+	};
+	TArray<FCustomFieldEntry> CustomFieldEntries;
 
+	// !!! If you edit this, be sure to update the upgrade paths in CookMetadata.cpp
 	friend FArchive& operator<<(FArchive& Ar, FCookMetadataPluginHierarchy& Hierarchy)
 	{
 		Ar << Hierarchy.PluginsEnabledAtCook << Hierarchy.PluginDependencies;
-		Ar << Hierarchy.RootPlugins << Hierarchy.CustomFieldNames;
+		Ar << Hierarchy.RootPlugins << Hierarchy.CustomFieldEntries;
 		return Ar;
 	}
 };
diff --git a/Engine/Source/Developer/IoStoreUtilities/Private/IoStoreUtilities.cpp b/Engine/Source/Developer/IoStoreUtilities/Private/IoStoreUtilities.cpp
--- a/Engine/Source/Developer/IoStoreUtilities/Private/IoStoreUtilities.cpp
+++ b/Engine/Source/Developer/IoStoreUtilities/Private/IoStoreUtilities.cpp
@@ -4099,15 +4099,24 @@ static void UpdatePluginMetadataAndWriteJsons(
 		if (InGraphEntry.IndexInEnabledPlugins != TNumericLimits<uint16>::Max())
 		{
 			const UE::Cook::FCookMetadataPluginEntry& CookMetadataEntry = MutablePluginHierarchy.PluginsEnabledAtCook[InGraphEntry.IndexInEnabledPlugins];
-			for (const TPair<uint8, bool>& BoolValue : CookMetadataEntry.CustomBoolFields)
+			for (const TPair<uint8, UE::Cook::FCookMetadataPluginEntry::CustomFieldVariantType>& CustomValue : CookMetadataEntry.CustomFields)
 			{
-				const FString& FieldName = MutablePluginHierarchy.CustomFieldNames[BoolValue.Key];
-				OutPluginMetadataJson << "\t\"" << FieldName << (BoolValue.Value ? "\":true,\n" : "\":false,\n");
-			}
-			for (const TPair<uint8, FString>& StringValue : CookMetadataEntry.CustomStringFields)
-			{
-				const FString& FieldName = MutablePluginHierarchy.CustomFieldNames[StringValue.Key];
-				OutPluginMetadataJson << "\t\"" << FieldName << "\":\"" << StringValue.Value << "\",\n";
+				const FString& FieldName = MutablePluginHierarchy.CustomFieldEntries[CustomValue.Key].Name;
+				UE::Cook::ECookMetadataCustomFieldType FieldType = MutablePluginHierarchy.CustomFieldEntries[CustomValue.Key].Type;
+
+				OutPluginMetadataJson << "\t\"" << FieldName;
+				
+				if (CustomValue.Value.IsType<bool>())
+				{
+					check(FieldType == UE::Cook::ECookMetadataCustomFieldType::Bool);
+					OutPluginMetadataJson << (CustomValue.Value.Get<bool>() ? "\":true,\n" : "\":false,\n");
+				}
+				else
+				{
+					check(FieldType == UE::Cook::ECookMetadataCustomFieldType::String);
+					OutPluginMetadataJson << "\":\"" << CustomValue.Value.Get<FString>() << "\",\n";
+				}
+
 			}
 		}
 
diff --git a/Engine/Source/Editor/UnrealEd/Private/CookOnTheFlyServer.cpp b/Engine/Source/Editor/UnrealEd/Private/CookOnTheFlyServer.cpp
--- a/Engine/Source/Editor/UnrealEd/Private/CookOnTheFlyServer.cpp
+++ b/Engine/Source/Editor/UnrealEd/Private/CookOnTheFlyServer.cpp
@@ -9616,9 +9616,10 @@ void UCookOnTheFlyServer::WriteCookMetadata(const ITargetPlatform* InTargetPlatf
 
 		// Get the names as a unique list for indexing the names for serialization
 		TArray<FString> CustomFieldNames;
+		TArray<UE::Cook::ECookMetadataCustomFieldType> CustomFieldTypes;
 		TMap<FString, uint8> CustomFieldNameIndex;
 		{
-			auto GetNames = [&CustomFieldNameIndex, &CustomFieldNames](const TArray<FString>& FieldList)
+			auto GetNames = [&CustomFieldNameIndex, &CustomFieldNames, &CustomFieldTypes](const TArray<FString>& FieldList, UE::Cook::ECookMetadataCustomFieldType FieldType)
 			{
 				for (const FString& FieldName : FieldList)
 				{
@@ -9626,15 +9627,16 @@ void UCookOnTheFlyServer::WriteCookMetadata(const ITargetPlatform* InTargetPlatf
 					if (FoundAtIndex == MAX_uint8)
 					{
 						CustomFieldNames.Add(FieldName);
+						CustomFieldTypes.Add(FieldType);
 						FoundAtIndex = (uint8)(CustomFieldNames.Num() - 1);
 					}
 				}
 			};
 
-			GetNames(BoolCustomFieldsList);
-			GetNames(StringCustomFieldsList);
-			GetNames(PerPlatformBoolCustomFieldsList);
-			GetNames(PerPlatformStringCustomFieldsList);
+			GetNames(BoolCustomFieldsList, UE::Cook::ECookMetadataCustomFieldType::Bool);
+			GetNames(StringCustomFieldsList, UE::Cook::ECookMetadataCustomFieldType::String);
+			GetNames(PerPlatformBoolCustomFieldsList, UE::Cook::ECookMetadataCustomFieldType::Bool);
+			GetNames(PerPlatformStringCustomFieldsList, UE::Cook::ECookMetadataCustomFieldType::String);
 		}
 
 		if (CustomFieldNames.Num() > 255)
@@ -9683,7 +9685,7 @@ void UCookOnTheFlyServer::WriteCookMetadata(const ITargetPlatform* InTargetPlatf
 			{
 				for (const FString& FieldName : FieldNames)
 				{
-					TVariant<bool, FString> FieldValue;
+					UE::Cook::FCookMetadataPluginEntry::CustomFieldVariantType FieldValue;
 			
 					bool bHasField = false;
 
@@ -9771,14 +9773,7 @@ void UCookOnTheFlyServer::WriteCookMetadata(const ITargetPlatform* InTargetPlatf
 
 					if (bHasField)
 					{
-						if (bIsBool)
-						{
-							Entry.CustomBoolFields.Add(CustomFieldNameIndex[FieldName], FieldValue.Get<bool>());
-						}
-						else
-						{
-							Entry.CustomStringFields.Add(CustomFieldNameIndex[FieldName], MoveTemp(FieldValue.Get<FString>()));
-						}
+						Entry.CustomFields.Add(CustomFieldNameIndex[FieldName], FieldValue);
 					}
 				} // end each field name
 
@@ -9826,8 +9821,13 @@ void UCookOnTheFlyServer::WriteCookMetadata(const ITargetPlatform* InTargetPlatf
 		PluginHierarchy.PluginsEnabledAtCook = MoveTemp(PluginsToAdd);
 		PluginHierarchy.PluginDependencies = MoveTemp(PluginChildArray);
 		PluginHierarchy.RootPlugins = MoveTemp(RootPlugins);
-		PluginHierarchy.CustomFieldNames = MoveTemp(CustomFieldNames);
 
+		for (int32 FieldIndex = 0; FieldIndex < CustomFieldNames.Num(); FieldIndex++)
+		{
+			UE::Cook::FCookMetadataPluginHierarchy::FCustomFieldEntry& NewEntry = PluginHierarchy.CustomFieldEntries.AddDefaulted_GetRef();
+			NewEntry.Name = MoveTemp(CustomFieldNames[FieldIndex]);
+			NewEntry.Type = CustomFieldTypes[FieldIndex];
+		}
 
 		// Sanity check we assigned plugin types
 		for (UE::Cook::FCookMetadataPluginEntry& Entry : PluginHierarchy.PluginsEnabledAtCook)
